Fixed rolling BP template pack

diff --git a/Templates/TP_RollingBP/Config/config.ini b/Templates/TP_RollingBP/Config/config.ini
--- a/Templates/TP_RollingBP/Config/config.ini
+++ b/Templates/TP_RollingBP/Config/config.ini
@@ -21,4 +21,7 @@
 [Redirects]
 GameName=TP_RollingBP
 
+[AdditionalFilesToAdd]
++Files=Templates/TP_RollingBP/Content/*.*
+
 
